Improve UV Layout BestRect heuristic toward minimal surface area instead of relying on lower left corners only

From dev-enterprise[REDACTED_HANDLE]

diff --git a/Engine/Source/Runtime/MeshUtilitiesCommon/Private/Allocator2D.cpp b/Engine/Source/Runtime/MeshUtilitiesCommon/Private/Allocator2D.cpp
--- a/Engine/Source/Runtime/MeshUtilitiesCommon/Private/Allocator2D.cpp
+++ b/Engine/Source/Runtime/MeshUtilitiesCommon/Private/Allocator2D.cpp
@@ -211,13 +211,12 @@ bool FAllocator2D::FindBitByBit( FRect& Rect, const FAllocator2D& Other )
 	return false;
 }
 
-bool FAllocator2D::FindWithSegments( FRect& Rect, const FRect& BestRect, const FAllocator2D& Other ) const
+bool FAllocator2D::FindWithSegments( FRect& Rect, const FAllocator2D& Other, TFunctionRef<bool (const FRect&)> IsBestRect) const
 {
 	FRect TestRect = Rect;
 
 	const uint32 MaxWidth    = Width - TestRect.W;
 	const uint32 MaxHeight   = Height - TestRect.H;
-	const uint32 BestArea    = BestRect.X + BestRect.Y * Width;
 
 	// For charts that are longer on the Y axis, its a lot faster to use
 	// another grid optimized for searching vertically because it will reduce
@@ -238,7 +237,7 @@ bool FAllocator2D::FindWithSegments( FRect& Rect, const FRect& BestRect, const F
 				Stats.FindWithSegmentsIterationsY++;
 
 				// Do not bother continue if another orientation had a better result
-				if( TestRect.X + TestRect.Y * Width >= BestArea)
+				if( !IsBestRect(Rect) )
 				{
 					Stats.FindWithSegmentsMovedPastPreviousBest++;
 					return false;
@@ -266,7 +265,7 @@ bool FAllocator2D::FindWithSegments( FRect& Rect, const FRect& BestRect, const F
 				Stats.FindWithSegmentsIterationsX++;
 
 				// Do not bother continue if another orientation had a better result
-				if( TestRect.X + TestRect.Y * Width >= BestArea)
+				if( !IsBestRect(Rect) )
 				{
 					Stats.FindWithSegmentsMovedPastPreviousBest++;
 					return false;
@@ -822,7 +821,7 @@ void FAllocator2D::ResetStats()
 #endif
 }
 
-void FAllocator2D::PublishStats(int32 ChartIndex, int32 Orientation, bool bFound, const FRect& Rect, const FRect& BestRect, const FMD5Hash& ChartMD5)
+void FAllocator2D::PublishStats(int32 ChartIndex, int32 Orientation, bool bFound, const FRect& Rect, const FRect& BestRect, const FMD5Hash& ChartMD5, TFunctionRef<bool (const FAllocator2D::FRect&)> IsBestRect)
 {
 	//This is super helpful in Insights to inspect long running charts behavior
 
@@ -865,7 +864,7 @@ void FAllocator2D::PublishStats(int32 ChartIndex, int32 Orientation, bool bFound
 			BestRect.Y,
 			BestRect.W,
 			BestRect.H,
-			(bFound && ((Rect.X + Rect.Y * Width) < (BestRect.X + BestRect.Y * Width))) ? 1 : 0,
+			(bFound && IsBestRect(Rect)) ? 1 : 0,
 			Stats.FindWithSegmentsIterationsY.GetValue(),
 			Stats.FindWithSegmentsIterationsX.GetValue(),
 			Stats.FindWithSegmentsMovedPastPreviousBest.GetValue(),
diff --git a/Engine/Source/Runtime/MeshUtilitiesCommon/Private/LayoutUV.cpp b/Engine/Source/Runtime/MeshUtilitiesCommon/Private/LayoutUV.cpp
--- a/Engine/Source/Runtime/MeshUtilitiesCommon/Private/LayoutUV.cpp
+++ b/Engine/Source/Runtime/MeshUtilitiesCommon/Private/LayoutUV.cpp
@@ -1059,6 +1059,27 @@ bool FLayoutUV::FChartPacker::PackCharts(TArray< FMeshChart >& Charts, float UVS
 			int32				BestOrientation = -1;
 			FAllocator2D::FRect	BestRect = { ~0u, ~0u, ~0u, ~0u };
 
+			// Refactored BestRect comparison code in one place so this can be customized per version if needed
+			TFunction<bool (const FAllocator2D::FRect&)> IsBestRect;
+			if ( LayoutVersion >= ELightmapUVVersion::OptimalSurfaceArea )
+			{
+				// This version focus on minimal surface area giving fairness to both horizontal and vertical chart placement
+				// instead of only taking the pixel offset of the lower left corner into account.
+				IsBestRect = 
+					[&BestRect](const FAllocator2D::FRect& Rect)
+					{
+						return ((Rect.X+Rect.W) + (Rect.Y+Rect.H)) < ((BestRect.X+BestRect.W) + (BestRect.Y+BestRect.H));
+					};
+			}
+			else
+			{
+				IsBestRect = 
+					[this, &BestRect](const FAllocator2D::FRect& Rect)
+					{
+						return Rect.X + Rect.Y * TextureResolution < BestRect.X + BestRect.Y * TextureResolution;
+					};
+			}
+
 			for( int32 Orientation = 0; Orientation < 8; Orientation++ )
 			{
 				// TODO If any dimension is less than 1 pixel shrink dimension to zero
@@ -1085,8 +1106,7 @@ bool FLayoutUV::FChartPacker::PackCharts(TArray< FMeshChart >& Charts, float UVS
 				{
 					if( LayoutRaster.Find( Rect ) )
 					{
-						// Is best?
-						if( Rect.X + Rect.Y * TextureResolution < BestRect.X + BestRect.Y * TextureResolution )
+						if( IsBestRect(Rect) )
 						{
 							BestOrientation = Orientation;
 							BestRect = Rect;
@@ -1150,7 +1170,7 @@ bool FLayoutUV::FChartPacker::PackCharts(TArray< FMeshChart >& Charts, float UVS
 							}
 
 							LayoutRaster.ResetStats();
-							bFound = LayoutRaster.FindWithSegments(RasterRect, BestRect, ChartRaster);
+							bFound = LayoutRaster.FindWithSegments(RasterRect, ChartRaster, IsBestRect);
 							if (bFound)
 							{
 								// Store only the best possible position in the hash table so we can start from there for other identical charts
@@ -1167,7 +1187,7 @@ bool FLayoutUV::FChartPacker::PackCharts(TArray< FMeshChart >& Charts, float UVS
 									Rect.X = RasterRect.X;
 									Rect.Y = RasterRect.Y;
 
-									bFound = LayoutRaster.FindWithSegments(Rect, BestRect, ChartRaster);
+									bFound = LayoutRaster.FindWithSegments(Rect, ChartRaster, IsBestRect);
 								}
 								else
 								{
@@ -1177,7 +1197,7 @@ bool FLayoutUV::FChartPacker::PackCharts(TArray< FMeshChart >& Charts, float UVS
 								}
 							}
 
-							LayoutRaster.PublishStats(ChartIndex, Orientation, bFound, Rect, BestRect, RasterMD5);
+							LayoutRaster.PublishStats(ChartIndex, Orientation, bFound, Rect, BestRect, RasterMD5, IsBestRect);
 						}
 
 					}
@@ -1190,8 +1210,7 @@ bool FLayoutUV::FChartPacker::PackCharts(TArray< FMeshChart >& Charts, float UVS
 
 					if( bFound )
 					{
-						// Is best?
-						if( Rect.X + Rect.Y * TextureResolution < BestRect.X + BestRect.Y * TextureResolution )
+						if( IsBestRect(Rect) )
 						{
 							BestChartRaster = ChartRaster;
 
diff --git a/Engine/Source/Runtime/MeshUtilitiesCommon/Public/Allocator2D.h b/Engine/Source/Runtime/MeshUtilitiesCommon/Public/Allocator2D.h
--- a/Engine/Source/Runtime/MeshUtilitiesCommon/Public/Allocator2D.h
+++ b/Engine/Source/Runtime/MeshUtilitiesCommon/Public/Allocator2D.h
@@ -66,7 +66,7 @@ public:
 	void       Alloc( FRect Rect );
 
 	bool       FindBitByBit( FRect& Rect, const FAllocator2D& Other );
-	bool       FindWithSegments( FRect& Rect, const FRect& BestRect, const FAllocator2D& Other ) const;
+	bool       FindWithSegments( FRect& Rect, const FAllocator2D& Other, TFunctionRef<bool (const FAllocator2D::FRect&)> IsBestRect ) const;
 	bool       Test( FRect Rect, const FAllocator2D& Other );
 	void       Alloc( FRect Rect, const FAllocator2D& Other );
 
@@ -98,7 +98,7 @@ public:
 	uint32     GetRasterHeight() const { return RasterHeight; }
 
 	void       ResetStats();
-	void       PublishStats( int32 ChartIndex, int32 Orientation, bool bFound, const FRect& Rect, const FRect& BestRect, const FMD5Hash& ChartMD5 );
+	void       PublishStats( int32 ChartIndex, int32 Orientation, bool bFound, const FRect& Rect, const FRect& BestRect, const FMD5Hash& ChartMD5, TFunctionRef<bool (const FAllocator2D::FRect&)> IsBestRect );
 
 protected:
 	bool       TestOneRun( const FRun& Run, const FRun& OtherRun, uint32 RectOffset, uint32 RectLength, uint32 PrimaryResolution, uint32& OutFailedLength ) const;
diff --git a/Engine/Source/Runtime/MeshUtilitiesCommon/Public/MeshUtilitiesCommon.h b/Engine/Source/Runtime/MeshUtilitiesCommon/Public/MeshUtilitiesCommon.h
--- a/Engine/Source/Runtime/MeshUtilitiesCommon/Public/MeshUtilitiesCommon.h
+++ b/Engine/Source/Runtime/MeshUtilitiesCommon/Public/MeshUtilitiesCommon.h
@@ -15,7 +15,8 @@ enum class ELightmapUVVersion : int32
 	ConsiderLightmapPadding = 6,
 	ForceLightmapPadding = 7,
 	Segments2D = 8,
-	Latest = Segments2D
+	OptimalSurfaceArea = 9,
+	Latest = OptimalSurfaceArea
 };
 
 /** Helper struct for building acceleration structures. */
