Delete a checkSlow that often fails

diff --git a/Engine/Source/Runtime/MeshUtilitiesCommon/Private/LayoutUV.cpp b/Engine/Source/Runtime/MeshUtilitiesCommon/Private/LayoutUV.cpp
--- a/Engine/Source/Runtime/MeshUtilitiesCommon/Private/LayoutUV.cpp
+++ b/Engine/Source/Runtime/MeshUtilitiesCommon/Private/LayoutUV.cpp
@@ -346,7 +346,6 @@ int32 FLayoutUV::FChartFinder::FindCharts( const FOverlappingCorners& Overlappin
 					
 					FVector2D EdgeOffset0 = UV0i - UV1j;
 					FVector2D EdgeOffset1 = UV1i - UV0j;
-					checkSlow( ( EdgeOffset0 - EdgeOffset1 ).IsNearlyZero(UVLAYOUT_THRESH_UVS_ARE_SAME) );
 
 					FVector2D Translation = EdgeOffset0;
 
