Temporarily remove overloaded function that is causing build failures.

diff --git a/Engine/Source/Runtime/MeshUtilitiesCommon/Public/MeshUtilitiesCommon.h b/Engine/Source/Runtime/MeshUtilitiesCommon/Public/MeshUtilitiesCommon.h
--- a/Engine/Source/Runtime/MeshUtilitiesCommon/Public/MeshUtilitiesCommon.h
+++ b/Engine/Source/Runtime/MeshUtilitiesCommon/Public/MeshUtilitiesCommon.h
@@ -67,16 +67,16 @@ inline bool PointsEqual(const FVector& V1, const FVector& V2, float ComparisonTh
 	return true;
 }
 
-inline bool PointsEqual(const FVector3f& V1, const FVector3f& V2, float ComparisonThreshold)
-{
-	if (FMath::Abs(V1.X - V2.X) > ComparisonThreshold
-		|| FMath::Abs(V1.Y - V2.Y) > ComparisonThreshold
-		|| FMath::Abs(V1.Z - V2.Z) > ComparisonThreshold)
-	{
-		return false;
-	}
-	return true;
-}
+// inline bool PointsEqual(const FVector3f& V1, const FVector3f& V2, float ComparisonThreshold)
+// {
+// 	if (FMath::Abs(V1.X - V2.X) > ComparisonThreshold
+// 		|| FMath::Abs(V1.Y - V2.Y) > ComparisonThreshold
+// 		|| FMath::Abs(V1.Z - V2.Z) > ComparisonThreshold)
+// 	{
+// 		return false;
+// 	}
+// 	return true;
+// }
 
 namespace TriangleUtilities
 {
