Material Baking - General fix for degenerate UVs
* When packing UVs, avoid relying on the UV area to scale triangles in the UV map. Instead, use the edges length (perimeter)

diff --git a/Engine/Source/Runtime/MeshUtilitiesCommon/Private/LayoutUV.cpp b/Engine/Source/Runtime/MeshUtilitiesCommon/Private/LayoutUV.cpp
--- a/Engine/Source/Runtime/MeshUtilitiesCommon/Private/LayoutUV.cpp
+++ b/Engine/Source/Runtime/MeshUtilitiesCommon/Private/LayoutUV.cpp
@@ -16,9 +16,9 @@ DEFINE_LOG_CATEGORY_STATIC(LogLayoutUV, Log, All);
 
 #define CHART_JOINING	1
 
-#define NEW_UVS_ARE_SAME THRESH_POINTS_ARE_SAME
-#define LEGACY_UVS_ARE_SAME (1.0f / 1024.0f)
-#define UVLAYOUT_THRESH_UVS_ARE_SAME (GetUVEqualityThreshold())
+#define NEW_UVS_ARE_SAME				UE_THRESH_POINTS_ARE_SAME
+#define LEGACY_UVS_ARE_SAME				UE_THRESH_UVS_ARE_SAME
+#define UVLAYOUT_THRESH_UVS_ARE_SAME	(GetUVEqualityThreshold())
 
 TAtomic<uint64> FLayoutUV::FindBestPackingCount(0);
 TAtomic<uint64> FLayoutUV::FindBestPackingCycles(0);
@@ -238,7 +238,8 @@ int32 FLayoutUV::FChartFinder::FindCharts( const FOverlappingCorners& Overlappin
 		Chart.MinUV = FVector2f( FLT_MAX, FLT_MAX );
 		Chart.MaxUV = FVector2f( -FLT_MAX, -FLT_MAX );
 		Chart.UVArea = 0.0f;
-		Chart.WorldScale = FVector2f::ZeroVector;
+		Chart.UVLengthSum = 0.0f;
+		Chart.WorldLengthSum = 0.0f;
 		FMemory::Memset( Chart.Join, 0xff );
 
 		Chart.FirstTri = Tri;
@@ -266,30 +267,41 @@ int32 FLayoutUV::FChartFinder::FindCharts( const FOverlappingCorners& Overlappin
 
 			FVector3f Edge1 = Positions[1] - Positions[0];
 			FVector3f Edge2 = Positions[2] - Positions[0];
-			float Area = 0.5f * ( Edge1 ^ Edge2 ).Size();
+			FVector Edge3 = Positions[2] - Positions[1];
+			float WorldLength = Edge1.Length() + Edge2.Length() + Edge3.Length();
 
 			FVector2f EdgeUV1 = UVs[1] - UVs[0];
 			FVector2f EdgeUV2 = UVs[2] - UVs[0];
+			FVector2f EdgeUV3 = UVs[2] - UVs[1];
+			float UVLength = EdgeUV1.Length() + EdgeUV2.Length() + EdgeUV3.Length();
 			float UVArea = 0.5f * FMath::Abs( EdgeUV1.X * EdgeUV2.Y - EdgeUV1.Y * EdgeUV2.X );
 
-			FVector2f UVLength;
-			UVLength.X = ( EdgeUV2.Y * Edge1 - EdgeUV1.Y * Edge2 ).Size();
-			UVLength.Y = (-EdgeUV2.X * Edge1 + EdgeUV1.X * Edge2 ).Size();
-			
-			Chart.WorldScale += UVLength;
+			Chart.UVLengthSum += UVLength;
+			Chart.WorldLengthSum += WorldLength;
 			Chart.UVArea += UVArea;
 		}
 		
 		Chart.LastTri = Tri;
 
 #if !CHART_JOINING
-		if (LayoutVersion >= ELightmapUVVersion::SmallChartPacking)
+		if (LayoutVersion >= ELightmapUVVersion::ScaleByEdgesLength)
+		{
+			if (Chart.UVLengthSum < UE_SMALL_NUMBER)
+			{
+				Chart.WorldScale = FVector2f(1.0f);
+			}
+			else
+			{
+				Chart.WorldScale = FVector2f(Chart.WorldLengthSum / Chart.UVLengthSum);
+			}
+		}
+		else if (LayoutVersion >= ELightmapUVVersion::SmallChartPacking)
 		{
-			Chart.WorldScale /= FMath::Max(Chart.UVArea, 1e-8f);
+			Chart.WorldScale /= FMath::Max(Chart.UVArea, UE_SMALL_NUMBER);
 		}
 		else
 		{
-			if (Chart.UVArea > 1e-4f)
+			if (Chart.UVArea > UE_KINDA_SMALL_NUMBER)
 			{
 				Chart.WorldScale /= Chart.UVArea;
 			}
@@ -561,11 +573,15 @@ int32 FLayoutUV::FChartFinder::FindCharts( const FOverlappingCorners& Overlappin
 						// Fixing joined chart MaxUV value to properly inflate non-joined axis extent
 						ChartA.MaxUV[ Axis ^ 1 ] = FMath::Max( ChartA.MaxUV[ Axis ^ 1 ], ChartA.MinUV[ Axis ^ 1 ] + ( ChartB.MaxUV[ Axis ^ 1 ] - ChartB.MinUV[ Axis ^ 1 ] ) );
 					}
+					ChartA.UVLengthSum += ChartB.UVLengthSum;
+					ChartA.WorldLengthSum += ChartB.WorldLengthSum;
 					ChartA.WorldScale += ChartB.WorldScale;
 					ChartA.UVArea += ChartB.UVArea;
 
 					ChartB.FirstTri = 0;
 					ChartB.LastTri = 0;
+					ChartB.UVLengthSum = 0.0f;
+					ChartB.WorldLengthSum = 0.0f;
 					ChartB.UVArea = 0.0f;
 
 					DisconnectChart( Charts, ChartB, Side ^ 2 );
@@ -608,13 +624,25 @@ int32 FLayoutUV::FChartFinder::FindCharts( const FOverlappingCorners& Overlappin
 	{
 		FMeshChart& Chart = Charts[i];
 
+		if (LayoutVersion >= ELightmapUVVersion::ScaleByEdgesLength)
+		{					
+			if (Chart.UVLengthSum < UE_SMALL_NUMBER)
+			{
+				Chart.WorldScale = FVector2f(1.0f);
+			}
+			else
+			{
+				Chart.WorldScale = FVector2f(Chart.WorldLengthSum / Chart.UVLengthSum);
+			}
+		}
+		else
 		if (LayoutVersion >= ELightmapUVVersion::SmallChartPacking)
 		{
-			Chart.WorldScale /= FMath::Max(Chart.UVArea, 1e-8f);
+			Chart.WorldScale /= FMath::Max(Chart.UVArea, UE_SMALL_NUMBER);
 		}
 		else
 		{
-			if (Chart.UVArea > 1e-4f)
+			if (Chart.UVArea > UE_KINDA_SMALL_NUMBER)
 			{
 				Chart.WorldScale /= Chart.UVArea;
 			}
diff --git a/Engine/Source/Runtime/MeshUtilitiesCommon/Public/LayoutUV.h b/Engine/Source/Runtime/MeshUtilitiesCommon/Public/LayoutUV.h
--- a/Engine/Source/Runtime/MeshUtilitiesCommon/Public/LayoutUV.h
+++ b/Engine/Source/Runtime/MeshUtilitiesCommon/Public/LayoutUV.h
@@ -23,6 +23,9 @@ struct FMeshChart
 	FVector2f	UVScale;
 	FVector2f	WorldScale;
 	
+	float		UVLengthSum;
+	float		WorldLengthSum;
+
 	FVector2f	PackingScaleU;
 	FVector2f	PackingScaleV;
 	FVector2f	PackingBias;
diff --git a/Engine/Source/Runtime/MeshUtilitiesCommon/Public/MeshUtilitiesCommon.h b/Engine/Source/Runtime/MeshUtilitiesCommon/Public/MeshUtilitiesCommon.h
--- a/Engine/Source/Runtime/MeshUtilitiesCommon/Public/MeshUtilitiesCommon.h
+++ b/Engine/Source/Runtime/MeshUtilitiesCommon/Public/MeshUtilitiesCommon.h
@@ -16,7 +16,8 @@ enum class ELightmapUVVersion : int32
 	ForceLightmapPadding = 7,
 	Segments2D = 8,
 	OptimalSurfaceArea = 9,
-	Latest = OptimalSurfaceArea
+	ScaleByEdgesLength = 10,
+	Latest = ScaleByEdgesLength
 };
 
 /**
diff --git a/Engine/Source/Runtime/StaticMeshDescription/Private/StaticMeshOperations.cpp b/Engine/Source/Runtime/StaticMeshDescription/Private/StaticMeshOperations.cpp
--- a/Engine/Source/Runtime/StaticMeshDescription/Private/StaticMeshOperations.cpp
+++ b/Engine/Source/Runtime/StaticMeshDescription/Private/StaticMeshOperations.cpp
@@ -1939,22 +1939,6 @@ struct FUniqueUVMeshDescriptionView final : FLayoutUV::IMeshView
 	TArray<FVector2D>&					OutputTexCoords;
 	bool								bMustRemap;
 
-	bool IsDegenerateTriangle(const FVector2f& A, const FVector2f& B, const FVector2f& C)
-	{
-		const FVector2f AB = B - A;
-		const FVector2f AC = C - A;
-		const float DoubleArea = FMath::Abs(AB ^ AC);
-		return DoubleArea < 2 * UE_SMALL_NUMBER;
-	}
-
-	bool IsDegenerateTriangle(const FVector3f& A, const FVector3f& B, const FVector3f& C)
-	{
-		const FVector3f AB = B - A;
-		const FVector3f AC = C - A;
-		const float DoubleAreaSquared = (AB ^ AC).SizeSquared();
-		return DoubleAreaSquared < 2 * UE_SMALL_NUMBER * UE_SMALL_NUMBER;
-	}
-
 	FUniqueUVMeshDescriptionView(const FMeshDescription& InMeshDescription, bool bMergeIdenticalMaterials, TArray<FVector2D>& InOutTexCoords)
 		: MeshDescription(InMeshDescription)
 		, OutputTexCoords(InOutTexCoords)
@@ -1998,39 +1982,33 @@ struct FUniqueUVMeshDescriptionView final : FLayoutUV::IMeshView
 
 			FTriangleID RemapTriangleID = TriangleID;
 
-			// Filter out degenerate triangles
-			bool bDegenerate = false;
-			bDegenerate = bDegenerate || IsDegenerateTriangle(Positions[VertexIDs[0]], Positions[VertexIDs[1]], Positions[VertexIDs[2]]);
-			bDegenerate = bDegenerate || IsDegenerateTriangle(TexCoords[VertexInstancesIDs[0]], TexCoords[VertexInstancesIDs[1]], TexCoords[VertexInstancesIDs[2]]);
-			if (!bDegenerate)
+			bool bUnique = true;
+
+			if (bMergeIdenticalMaterials)
 			{
-				bool bUnique = true;
-				if (bMergeIdenticalMaterials)
+				int32 TriangleHash = GetTypeHash(RefPolygonGroupID);
+				for (const FVertexInstanceID& VertexInstanceID : VertexInstancesIDs)
 				{
-					int32 TriangleHash = GetTypeHash(RefPolygonGroupID);
-					for (const FVertexInstanceID& VertexInstanceID : VertexInstancesIDs)
-					{
-						// Compute hash based on UVs & vertices colors
-						HashAttribute(VertexInstanceID, TexCoords, TriangleHash);
-						HashAttribute(VertexInstanceID, VertexColors, TriangleHash);
-					}
-
-					FTriangleID* UniqueTriangleIDPtr = UniqueTriangles.Find(TriangleHash);
-					if (UniqueTriangleIDPtr != nullptr)
-					{
-						RemapTriangleID = *UniqueTriangleIDPtr;
-						bUnique = false;
-					}
-					else
-					{
-						UniqueTriangles.Add(TriangleHash, TriangleID);
-					}
+					// Compute hash based on UVs & vertices colors
+					HashAttribute(VertexInstanceID, TexCoords, TriangleHash);
+					HashAttribute(VertexInstanceID, VertexColors, TriangleHash);
 				}
 
-				if (bUnique)
+				FTriangleID* UniqueTriangleIDPtr = UniqueTriangles.Find(TriangleHash);
+				if (UniqueTriangleIDPtr != nullptr)
 				{
-					UniqueVerts.Append(VertexInstancesIDs);
+					RemapTriangleID = *UniqueTriangleIDPtr;
+					bUnique = false;
 				}
+				else
+				{
+					UniqueTriangles.Add(TriangleHash, TriangleID);
+				}
+			}
+
+			if (bUnique)
+			{
+				UniqueVerts.Append(VertexInstancesIDs);
 			}
 
 			RemapTriangles.Add(RemapTriangleID);
