Build fix: FVector -> FVector3f
#rnx

diff --git a/Engine/Source/Runtime/MeshUtilitiesCommon/Private/LayoutUV.cpp b/Engine/Source/Runtime/MeshUtilitiesCommon/Private/LayoutUV.cpp
--- a/Engine/Source/Runtime/MeshUtilitiesCommon/Private/LayoutUV.cpp
+++ b/Engine/Source/Runtime/MeshUtilitiesCommon/Private/LayoutUV.cpp
@@ -267,7 +267,7 @@ int32 FLayoutUV::FChartFinder::FindCharts( const FOverlappingCorners& Overlappin
 
 			FVector3f Edge1 = Positions[1] - Positions[0];
 			FVector3f Edge2 = Positions[2] - Positions[0];
-			FVector Edge3 = Positions[2] - Positions[1];
+			FVector3f Edge3 = Positions[2] - Positions[1];
 			float WorldLength = Edge1.Length() + Edge2.Length() + Edge3.Length();
 
 			FVector2f EdgeUV1 = UVs[1] - UVs[0];
