LayoutUV - Fix support for LayoutVersion < ELightmapUVVersion::ScaleByEdgesLength
Bumped DDC key to fix wrong cached data that could have been stored before this fix

diff --git a/Engine/Source/Runtime/Core/Private/UObject/FortniteMainBranchObjectVersion.cpp b/Engine/Source/Runtime/Core/Private/UObject/FortniteMainBranchObjectVersion.cpp
--- a/Engine/Source/Runtime/Core/Private/UObject/FortniteMainBranchObjectVersion.cpp
+++ b/Engine/Source/Runtime/Core/Private/UObject/FortniteMainBranchObjectVersion.cpp
@@ -13,6 +13,6 @@ TMap<FGuid, FGuid> FFortniteMainBranchObjectVersion::GetSystemGuids()
 	SystemGuids.Add(DevGuids.NIAGARASHADERMAP_DERIVEDDATA_VER, FGuid("2289D5116CF94BC0AFEAEC541468E645"));
 	SystemGuids.Add(DevGuids.Niagara_LatestScriptCompileVersion, FGuid("9CFE679BEE95F345BBC06744597AC064"));
 	SystemGuids.Add(DevGuids.SkeletalMeshDerivedDataVersion, FGuid("ACF593EAAE354FCBB34CF44F5AA1BFD2"));
-	SystemGuids.Add(DevGuids.STATICMESH_DERIVEDDATA_VER, FGuid("F3F4E2FE7EE14CE7942B445FABD05E37"));
+	SystemGuids.Add(DevGuids.STATICMESH_DERIVEDDATA_VER, FGuid("543843B103794E0E9BA4BE80FB602F79"));
 	return SystemGuids;
 }
diff --git a/Engine/Source/Runtime/MeshUtilitiesCommon/Private/LayoutUV.cpp b/Engine/Source/Runtime/MeshUtilitiesCommon/Private/LayoutUV.cpp
--- a/Engine/Source/Runtime/MeshUtilitiesCommon/Private/LayoutUV.cpp
+++ b/Engine/Source/Runtime/MeshUtilitiesCommon/Private/LayoutUV.cpp
@@ -238,6 +238,7 @@ int32 FLayoutUV::FChartFinder::FindCharts( const FOverlappingCorners& Overlappin
 		Chart.MinUV = FVector2f( FLT_MAX, FLT_MAX );
 		Chart.MaxUV = FVector2f( -FLT_MAX, -FLT_MAX );
 		Chart.UVArea = 0.0f;
+		Chart.WorldScale = FVector2f::ZeroVector;
 		Chart.UVLengthSum = 0.0f;
 		Chart.WorldLengthSum = 0.0f;
 		FMemory::Memset( Chart.Join, 0xff );
@@ -273,14 +274,28 @@ int32 FLayoutUV::FChartFinder::FindCharts( const FOverlappingCorners& Overlappin
 			FVector2f EdgeUV1 = UVs[1] - UVs[0];
 			FVector2f EdgeUV2 = UVs[2] - UVs[0];
 			FVector2f EdgeUV3 = UVs[2] - UVs[1];
-			float UVLength = EdgeUV1.Length() + EdgeUV2.Length() + EdgeUV3.Length();
-			float UVArea = 0.5f * FMath::Abs( EdgeUV1.X * EdgeUV2.Y - EdgeUV1.Y * EdgeUV2.X );
 
-			Chart.UVLengthSum += UVLength;
-			Chart.WorldLengthSum += WorldLength;
+			float UVArea = 0.5f * FMath::Abs(EdgeUV1.X * EdgeUV2.Y - EdgeUV1.Y * EdgeUV2.X);
 			Chart.UVArea += UVArea;
+
+			if (LayoutVersion >= ELightmapUVVersion::ScaleByEdgesLength)
+			{
+				float WorldLength = Edge1.Length() + Edge2.Length() + Edge3.Length();
+				float UVLength = EdgeUV1.Length() + EdgeUV2.Length() + EdgeUV3.Length();
+
+				Chart.UVLengthSum += UVLength;
+				Chart.WorldLengthSum += WorldLength;
+			}
+			else
+			{
+				FVector2f UVLength;
+				UVLength.X = (EdgeUV2.Y * Edge1 - EdgeUV1.Y * Edge2).Size();
+				UVLength.Y = (-EdgeUV2.X * Edge1 + EdgeUV1.X * Edge2).Size();
+
+				Chart.WorldScale += UVLength;
+			}
 		}
-		
+
 		Chart.LastTri = Tri;
 
 #if !CHART_JOINING
