Fixed mismatched copy constructor/assignment operator on FAllocator2D.

diff --git a/Engine/Source/Runtime/MeshUtilitiesCommon/Public/Allocator2D.h b/Engine/Source/Runtime/MeshUtilitiesCommon/Public/Allocator2D.h
--- a/Engine/Source/Runtime/MeshUtilitiesCommon/Public/Allocator2D.h
+++ b/Engine/Source/Runtime/MeshUtilitiesCommon/Public/Allocator2D.h
@@ -90,7 +90,8 @@ public:
 
 	void       CopyRuns( TArray<FRun>& Runs, const TArray<FRun>& OtherRuns, int32 MaxSize );
 
-	// Take control of the copy to reduce the amount of data movement to the strict minimum
+	// Take control of the copy assignment to reduce the amount of data movement to the strict minimum
+	FAllocator2D(const FAllocator2D& Other) = default;
 	FAllocator2D& operator = (const FAllocator2D& Other);
 
 	// Allow to visualize the content in ascii for debugging purpose. (i.e Watch or Immediate window).
